HlMedia: Use PrivateIncludePaths for HLMediaLibrary.h include

#rnx

diff --git a/Engine/Plugins/Media/HLMedia/Source/HLMedia/HLMedia.Build.cs b/Engine/Plugins/Media/HLMedia/Source/HLMedia/HLMedia.Build.cs
--- a/Engine/Plugins/Media/HLMedia/Source/HLMedia/HLMedia.Build.cs
+++ b/Engine/Plugins/Media/HLMedia/Source/HLMedia/HLMedia.Build.cs
@@ -42,8 +42,6 @@ public class HLMedia : ModuleRules
                 "HLMedia/Private",
             });
 
-        var EngineDir = Path.GetFullPath(Target.RelativeEnginePath);
-
         PrivateDependencyModuleNames.AddRange(new string[] 
         {
             "HLMediaLibrary",
@@ -51,7 +49,7 @@ public class HLMedia : ModuleRules
 			"D3D12RHI"
         });
 
-        PrivateIncludePaths.Add(Path.Combine(EngineDir, "Source/ThirdParty/HLMediaLibrary/inc"));
+        PrivateIncludePaths.Add(Path.Combine(EngineDirectory, "Source/ThirdParty/HLMediaLibrary/inc"));
 
         PublicSystemLibraries.Add("mfplat.lib");
         PublicSystemLibraries.Add("mfreadwrite.lib");
diff --git a/Engine/Plugins/Media/HLMedia/Source/HLMedia/Private/HLMediaPrivate.h b/Engine/Plugins/Media/HLMedia/Source/HLMedia/Private/HLMediaPrivate.h
--- a/Engine/Plugins/Media/HLMedia/Source/HLMedia/Private/HLMediaPrivate.h
+++ b/Engine/Plugins/Media/HLMedia/Source/HLMedia/Private/HLMediaPrivate.h
@@ -44,7 +44,7 @@ THIRD_PARTY_INCLUDES_START
 
 #pragma warning(push)
 #pragma warning(disable : 4005)
-#include "HLMediaLibrary/inc/HLMediaLibrary.h"
+#include "HLMediaLibrary.h"
 #pragma warning(pop)
 
 THIRD_PARTY_INCLUDES_END
diff --git a/Engine/Plugins/Media/HLMedia/Source/HLMediaEditor/HLMediaEditor.Build.cs b/Engine/Plugins/Media/HLMedia/Source/HLMediaEditor/HLMediaEditor.Build.cs
--- a/Engine/Plugins/Media/HLMedia/Source/HLMediaEditor/HLMediaEditor.Build.cs
+++ b/Engine/Plugins/Media/HLMedia/Source/HLMediaEditor/HLMediaEditor.Build.cs
@@ -1,5 +1,8 @@
 // Copyright Epic Games, Inc. All Rights Reserved.
 
+using System.IO;
+using UnrealBuildTool;
+
 namespace UnrealBuildTool.Rules
 {
 	public class HLMediaEditor : ModuleRules
@@ -21,7 +24,8 @@ namespace UnrealBuildTool.Rules
             PrivateIncludePaths.AddRange(
                 new string[] {
                     "HLMedia/Private",
-                });
+					Path.Combine(EngineDirectory, "Source/ThirdParty/HLMediaLibrary/inc"),
+		});
         }
     }
 }
