Removing hardcoded paths to D3D11RHI

diff --git a/Engine/Plugins/Media/HLMedia/Source/HLMedia/HLMedia.Build.cs b/Engine/Plugins/Media/HLMedia/Source/HLMedia/HLMedia.Build.cs
--- a/Engine/Plugins/Media/HLMedia/Source/HLMedia/HLMedia.Build.cs
+++ b/Engine/Plugins/Media/HLMedia/Source/HLMedia/HLMedia.Build.cs
@@ -1,6 +1,5 @@
 // Copyright Epic Games, Inc. All Rights Reserved.
 
-using System;
 using System.IO;
 using UnrealBuildTool;
 
