Minor HLMedia dependency cleanup.

diff --git a/Engine/Plugins/Media/HLMedia/Source/HLMedia/Private/HLMediaSource.cpp b/Engine/Plugins/Media/HLMedia/Source/HLMedia/Private/HLMediaSource.cpp
--- a/Engine/Plugins/Media/HLMedia/Source/HLMedia/Private/HLMediaSource.cpp
+++ b/Engine/Plugins/Media/HLMedia/Source/HLMedia/Private/HLMediaSource.cpp
@@ -2,6 +2,7 @@
 
 #include "HLMediaSource.h"
 
+#include "HLMediaPrivate.h"
 #include "Templates/SharedPointer.h"
 
 namespace HLMediaSource
diff --git a/Engine/Plugins/Media/HLMedia/Source/HLMedia/Private/HLMediaSource.h b/Engine/Plugins/Media/HLMedia/Source/HLMedia/Private/HLMediaSource.h
--- a/Engine/Plugins/Media/HLMedia/Source/HLMedia/Private/HLMediaSource.h
+++ b/Engine/Plugins/Media/HLMedia/Source/HLMedia/Private/HLMediaSource.h
@@ -2,8 +2,6 @@
 
 #pragma once
 
-#include "HLMediaPrivate.h"
-
 #include "BaseMediaSource.h"
 
 #include "HLMediaSource.generated.h"
diff --git a/Engine/Plugins/Media/HLMedia/Source/HLMediaEditor/Private/HLMediaEditorModule.cpp b/Engine/Plugins/Media/HLMedia/Source/HLMediaEditor/Private/HLMediaEditorModule.cpp
--- a/Engine/Plugins/Media/HLMedia/Source/HLMediaEditor/Private/HLMediaEditorModule.cpp
+++ b/Engine/Plugins/Media/HLMedia/Source/HLMediaEditor/Private/HLMediaEditorModule.cpp
@@ -1,12 +1,6 @@
 // Copyright Epic Games, Inc. All Rights Reserved.
 
 #include "Modules/ModuleInterface.h"
-#include "Modules/ModuleManager.h"
-#include "PropertyEditorModule.h"
-#include "UObject/NameTypes.h"
-
-#include "HLMediaSource.h"
-
 
 /**
  * Implements the HLMediaEditor module.
@@ -19,17 +13,11 @@ public:
 	//~ IModuleInterface interface
 	virtual void StartupModule() override
 	{
-		HLMediaSourceName = UHLMediaSource::StaticClass()->GetFName();
 	}
 
 	virtual void ShutdownModule() override
 	{
 	}
-
-private:
-
-	/** Class names. */
-	FName HLMediaSourceName;
 };
 
 
