Restoring required include.

diff --git a/Engine/Plugins/Media/HLMedia/Source/HLMediaEditor/Private/HLMediaEditorModule.cpp b/Engine/Plugins/Media/HLMedia/Source/HLMediaEditor/Private/HLMediaEditorModule.cpp
--- a/Engine/Plugins/Media/HLMedia/Source/HLMediaEditor/Private/HLMediaEditorModule.cpp
+++ b/Engine/Plugins/Media/HLMedia/Source/HLMediaEditor/Private/HLMediaEditorModule.cpp
@@ -1,6 +1,7 @@
 // Copyright Epic Games, Inc. All Rights Reserved.
 
 #include "Modules/ModuleInterface.h"
+#include "Modules/ModuleManager.h"
 
 /**
  * Implements the HLMediaEditor module.
