Fixing HLMedia include paths.

diff --git a/Engine/Platforms/Hololens/Plugins/Media/HLMedia/Source/HLMedia/HLMedia_HoloLens.Build.cs b/Engine/Platforms/Hololens/Plugins/Media/HLMedia/Source/HLMedia/HLMedia_HoloLens.Build.cs
--- a/Engine/Platforms/Hololens/Plugins/Media/HLMedia/Source/HLMedia/HLMedia_HoloLens.Build.cs
+++ b/Engine/Platforms/Hololens/Plugins/Media/HLMedia/Source/HLMedia/HLMedia_HoloLens.Build.cs
@@ -8,6 +8,7 @@ namespace UnrealBuildTool.Rules
 	{
 		public HLMedia_HoloLens(ReadOnlyTargetRules Target) : base(Target)
 		{
+			PrivateIncludePathModuleNames.Add("D3D11RHI_HoloLens");
 		}
 	}
 }
diff --git a/Engine/Plugins/Media/HLMedia/Source/HLMedia/Internal/HLMediaSource.h b/Engine/Plugins/Media/HLMedia/Source/HLMedia/Internal/HLMediaSource.h
new file mode 100644
--- /dev/null
+++ b/Engine/Plugins/Media/HLMedia/Source/HLMedia/Internal/HLMediaSource.h
@@ -0,0 +1,37 @@
+// Copyright Epic Games, Inc. All Rights Reserved.
+
+#pragma once
+
+#include "BaseMediaSource.h"
+
+#include "HLMediaSource.generated.h"
+
+UCLASS(BlueprintType)
+class HLMEDIA_API UHLMediaSource
+	: public UBaseMediaSource
+{
+	GENERATED_BODY()
+
+public:
+	/** Default constructor. */
+	UHLMediaSource();
+
+	/** The URL property is an Adaptive Streaming playlist. */
+	UPROPERTY(BlueprintReadWrite, EditAnywhere, Category = Stream, AdvancedDisplay)
+	bool IsAdaptiveSource;
+
+	/** The URL to the media stream to be played. */
+	UPROPERTY(BlueprintReadWrite, EditAnywhere, Category = Stream, AssetRegistrySearchable)
+	FString StreamUrl;
+
+	/**  IMediaOptions */
+	virtual FName GetDesiredPlayerName() const override;
+	virtual bool GetMediaOption(const FName& Key, bool DefaultValue) const override;
+	virtual bool HasMediaOption(const FName& Key) const override;
+
+public:
+
+	//~ UMediaSource interface
+	virtual FString GetUrl() const override;
+	virtual bool Validate() const override;
+};
diff --git a/Engine/Plugins/Media/HLMedia/Source/HLMedia/Private/HLMediaSource.h b/Engine/Plugins/Media/HLMedia/Source/HLMedia/Private/HLMediaSource.h
deleted file mode 100644
diff --git a/Engine/Plugins/Media/HLMedia/Source/HLMediaEditor/HLMediaEditor.Build.cs b/Engine/Plugins/Media/HLMedia/Source/HLMediaEditor/HLMediaEditor.Build.cs
--- a/Engine/Plugins/Media/HLMedia/Source/HLMediaEditor/HLMediaEditor.Build.cs
+++ b/Engine/Plugins/Media/HLMedia/Source/HLMediaEditor/HLMediaEditor.Build.cs
@@ -22,11 +22,6 @@ namespace UnrealBuildTool.Rules
                 });
 
 			PrivateIncludePathModuleNames.Add("HLMediaLibrary");
-
-            PrivateIncludePaths.AddRange(
-                new string[] {
-                    "HLMedia/Private",
-		});
         }
     }
 }
